Move DNACalib under public plugins folder and add DNACalib2 under restricted folder

diff --git a/Engine/Plugins/Animation/DNACalib/DNACalib.uplugin b/Engine/Plugins/Animation/DNACalib/DNACalib.uplugin
new file mode 100644
--- /dev/null
+++ b/Engine/Plugins/Animation/DNACalib/DNACalib.uplugin
@@ -0,0 +1,43 @@
+{
+	"FileVersion": 3,
+	"Version": 6.12,
+	"VersionName": "6.12.2",
+	"FriendlyName": "DNACalib Plugin v6.12.2",
+	"Description": "3Lateral DNA Calibration tool plugin",
+	"Category": "Animation",
+	"CreatedBy": "3Lateral",
+	"CreatedByURL": "http://www.3lateral.com/",
+	"DocsURL": "",
+	"MarketplaceURL": "",
+	"SupportURL": "",
+	"CanContainContent": false,
+	"IsBetaVersion": true,
+	"IsExperimentalVersion": false,
+	"Installed": false,
+	"Modules": [
+		{
+			"Name": "DNACalibLib",
+			"Type": "Runtime",
+			"LoadingPhase": "PreDefault"
+		},
+		{
+			"Name": "DNACalibLibTest",
+			"Type": "Runtime",
+			"LoadingPhase": "PreDefault",
+			"PlatformAllowList": [
+				"Win64"
+			]
+		},
+		{
+			"Name": "DNACalibModule",
+			"Type": "Runtime",
+			"LoadingPhase": "Default"
+		}
+	],
+	"Plugins": [
+		{
+			"Name": "RigLogic",
+			"Enabled": true
+		}
+	]
+}
\ No newline at end of file
diff --git a/Engine/Plugins/Animation/DNACalib/Source/DNACalibLib/Private/dnacalib/commands/SupportFactories.h b/Engine/Plugins/Animation/DNACalib/Source/DNACalibLib/Private/dnacalib/commands/SupportFactories.h
new file mode 100644
--- /dev/null
+++ b/Engine/Plugins/Animation/DNACalib/Source/DNACalibLib/Private/dnacalib/commands/SupportFactories.h
@@ -0,0 +1,73 @@
+// Copyright Epic Games, Inc. All Rights Reserved.
+
+#pragma once
+
+#include "dnacalib/TypeDefs.h"
+#include "dnacalib/commands/VectorOperations.h"
+#include "dnacalib/utils/Extd.h"
+
+#ifdef _MSC_VER
+    #pragma warning(push)
+    #pragma warning(disable : 4365 4987)
+#endif
+#include <cassert>
+#include <cstddef>
+#include <functional>
+#ifdef _MSC_VER
+    #pragma warning(pop)
+#endif
+
+namespace dnac {
+
+struct WeightGetterFactory {
+    using GetterFunc = std::function<float (float*, std::size_t)>;
+
+    static GetterFunc create(const Vector<float>& masks) {
+        if (masks.empty()) {
+            return [](float*  /*unused*/, std::size_t  /*unused*/) {
+                       return 1.0f;
+            };
+        }
+        return [](float* pMasks, std::size_t index) {
+                   return pMasks[index];
+        };
+    }
+
+};
+
+struct OperationFactory {
+    using OpFunc = std::function<float (float, float, float)>;
+
+    static OpFunc create(VectorOperation operation) {
+        switch (operation) {
+            case VectorOperation::Interpolate: {
+                return [](float a, float b, float weight) {
+                           return extd::interpolate(a, b, weight);
+                };
+            }
+            case VectorOperation::Add: {
+                return [](float a, float b, float weight) {
+                           return a + (b * weight);
+                };
+            }
+            case VectorOperation::Subtract: {
+                return [](float a, float b, float weight) {
+                           return a - (b * weight);
+                };
+            }
+            case VectorOperation::Multiply: {
+                return [](float a, float b, float weight) {
+                           return a * (b * weight);
+                };
+            }
+            default: {
+                // Unreachable, but needed to silence gcc warnings
+                assert(false);
+                return {};
+            }
+        }
+    }
+
+};
+
+}  // namespace dnac
diff --git a/Engine/Plugins/Animation/DNACalib/Source/DNACalibLib/Private/dnacalib/utils/ScopedEnumEx.h b/Engine/Plugins/Animation/DNACalib/Source/DNACalibLib/Private/dnacalib/utils/ScopedEnumEx.h
new file mode 100644
--- /dev/null
+++ b/Engine/Plugins/Animation/DNACalib/Source/DNACalibLib/Private/dnacalib/utils/ScopedEnumEx.h
@@ -0,0 +1,61 @@
+// Copyright Epic Games, Inc. All Rights Reserved.
+
+#pragma once
+
+#include <type_traits>
+
+namespace dnac {
+
+template<typename TEnum>
+typename std::enable_if<std::is_enum<TEnum>::value, TEnum>::type
+operator&(TEnum lhs, TEnum rhs) {
+    using Underlying = typename std::underlying_type<TEnum>::type;
+    return static_cast<TEnum>(static_cast<Underlying>(lhs) & static_cast<Underlying>(rhs));
+}
+
+template<typename TEnum>
+typename std::enable_if<std::is_enum<TEnum>::value, TEnum>::type
+operator|(TEnum lhs, TEnum rhs) {
+    using Underlying = typename std::underlying_type<TEnum>::type;
+    return static_cast<TEnum>(static_cast<Underlying>(lhs) | static_cast<Underlying>(rhs));
+}
+
+template<typename TEnum>
+typename std::enable_if<std::is_enum<TEnum>::value, TEnum>::type
+operator^(TEnum lhs, TEnum rhs) {
+    using Underlying = typename std::underlying_type<TEnum>::type;
+    return static_cast<TEnum>(static_cast<Underlying>(lhs) ^ static_cast<Underlying>(rhs));
+}
+
+template<typename TEnum>
+typename std::enable_if<std::is_enum<TEnum>::value, TEnum>::type
+operator~(TEnum value) {
+    using Underlying = typename std::underlying_type<TEnum>::type;
+    return static_cast<TEnum>(~static_cast<Underlying>(value));
+}
+
+template<typename TEnum>
+typename std::enable_if<std::is_enum<TEnum>::value, TEnum>::type
+operator&=(TEnum& lhs, TEnum rhs) {
+    return lhs = (lhs & rhs);
+}
+
+template<typename TEnum>
+typename std::enable_if<std::is_enum<TEnum>::value, TEnum>::type
+operator|=(TEnum& lhs, TEnum rhs) {
+    return lhs = (lhs | rhs);
+}
+
+template<typename TEnum>
+typename std::enable_if<std::is_enum<TEnum>::value, TEnum>::type
+operator^=(TEnum& lhs, TEnum rhs) {
+    return lhs = (lhs ^ rhs);
+}
+
+template<typename TEnum>
+typename std::enable_if<std::is_enum<TEnum>::value, bool>::type
+contains(TEnum lhs, TEnum rhs) {
+    return (lhs & rhs) == rhs;
+}
+
+}  // namespace dnac
diff --git a/Engine/Plugins/Animation/DNACalib/Source/DNACalibLib/Public/dnacalib/Command.h b/Engine/Plugins/Animation/DNACalib/Source/DNACalibLib/Public/dnacalib/Command.h
new file mode 100644
--- /dev/null
+++ b/Engine/Plugins/Animation/DNACalib/Source/DNACalibLib/Public/dnacalib/Command.h
@@ -0,0 +1,22 @@
+// Copyright Epic Games, Inc. All Rights Reserved.
+
+#pragma once
+
+#include "dnacalib/Defs.h"
+#include "dnacalib/types/Aliases.h"
+
+namespace dnac {
+
+class DNACalibDNAReader;
+
+/**
+    @brief Command is an abstract class whose implementations are expected to modify the DNA provided in the run() method in some way.
+*/
+class DNACAPI Command {
+    public:
+        virtual ~Command();
+        virtual void run(DNACalibDNAReader* output) = 0;
+
+};
+
+}  // namespace dnac
diff --git a/Engine/Plugins/Animation/DNACalib/Source/DNACalibLib/Public/dnacalib/version/Version.h b/Engine/Plugins/Animation/DNACalib/Source/DNACalibLib/Public/dnacalib/version/Version.h
new file mode 100644
--- /dev/null
+++ b/Engine/Plugins/Animation/DNACalib/Source/DNACalibLib/Public/dnacalib/version/Version.h
@@ -0,0 +1,8 @@
+// Copyright Epic Games, Inc. All Rights Reserved.
+
+#pragma once
+
+#define DNAC_MAJOR_VERSION 6
+#define DNAC_MINOR_VERSION 12
+#define DNAC_PATCH_VERSION 2
+#define DNAC_VERSION_STRING "6.12.2"
diff --git a/Engine/Plugins/Animation/DNACalib/Source/DNACalibLibTest/Private/dnactests/Defs.h b/Engine/Plugins/Animation/DNACalib/Source/DNACalibLibTest/Private/dnactests/Defs.h
new file mode 100644
--- /dev/null
+++ b/Engine/Plugins/Animation/DNACalib/Source/DNACalibLibTest/Private/dnactests/Defs.h
@@ -0,0 +1,69 @@
+// Copyright Epic Games, Inc. All Rights Reserved.
+
+#pragma once
+
+#ifdef _MSC_VER
+    #pragma warning(disable : 6326)
+
+    #pragma warning(push)
+    #pragma warning(disable : 4061 4121 4365 4987 4668 4996 6330)
+#endif
+#ifdef __clang__
+    #pragma clang diagnostic push
+    #pragma clang diagnostic ignored "-Wundef"
+#endif
+#include <gtest/gtest.h>
+#ifdef __clang__
+    #pragma clang diagnostic pop
+    // This cannot be applied to the scope of the header inclusion only, as the TEST macros will trigger it
+    #pragma clang diagnostic ignored "-Wglobal-constructors"
+    #pragma clang diagnostic ignored "-Wexit-time-destructors"
+    #pragma clang diagnostic ignored "-Wweak-vtables"
+#endif
+#ifdef _MSC_VER
+    #pragma warning(pop)
+#endif
+
+#ifndef INSTANTIATE_TEST_SUITE_P
+    #define INSTANTIATE_TEST_SUITE_P INSTANTIATE_TEST_CASE_P
+#endif
+
+#ifndef TYPED_TEST_SUITE_P
+    #define TYPED_TEST_SUITE_P TYPED_TEST_CASE_P
+#endif
+
+#ifndef REGISTER_TYPED_TEST_SUITE_P
+    #define REGISTER_TYPED_TEST_SUITE_P REGISTER_TYPED_TEST_CASE_P
+#endif
+
+#ifndef INSTANTIATE_TYPED_TEST_SUITE_P
+    #define INSTANTIATE_TYPED_TEST_SUITE_P INSTANTIATE_TYPED_TEST_CASE_P
+#endif
+
+#ifndef ASSERT_ELEMENTS_EQ
+    #define ASSERT_ELEMENTS_EQ(result, expected, count)     \
+        for (std::size_t i = 0ul; i < count; ++i) {         \
+            ASSERT_EQ(result[i], expected[i]);              \
+        }
+#endif
+
+#ifndef EXPECT_ELEMENTS_EQ
+    #define EXPECT_ELEMENTS_EQ(result, expected, count)     \
+        for (std::size_t i = 0ul; i < count; ++i) {         \
+            EXPECT_EQ(result[i], expected[i]);              \
+        }
+#endif
+
+#ifndef ASSERT_ELEMENTS_NEAR
+    #define ASSERT_ELEMENTS_NEAR(result, expected, count, threshold)    \
+        for (std::size_t i = 0ul; i < count; ++i) {                     \
+            ASSERT_NEAR(result[i], expected[i], threshold);             \
+        }
+#endif
+
+#ifndef EXPECT_ELEMENTS_NEAR
+    #define EXPECT_ELEMENTS_NEAR(result, expected, count, threshold)    \
+        for (std::size_t i = 0ul; i < count; ++i) {                     \
+            EXPECT_NEAR(result[i], expected[i], threshold);             \
+        }
+#endif
diff --git a/Engine/Plugins/Animation/DNACalib/Source/DNACalibModule/Public/DNACalibDNAReader.h b/Engine/Plugins/Animation/DNACalib/Source/DNACalibModule/Public/DNACalibDNAReader.h
new file mode 100644
--- /dev/null
+++ b/Engine/Plugins/Animation/DNACalib/Source/DNACalibModule/Public/DNACalibDNAReader.h
@@ -0,0 +1,188 @@
+// Copyright Epic Games, Inc. All Rights Reserved.
+
+#pragma once
+
+#include "DNAReader.h"
+
+#include "CoreMinimal.h"
+#include "Math/Vector.h"
+#include "UObject/NoExportTypes.h"
+#include "UObject/ObjectMacros.h"
+
+/**
+	@brief A UE wrapper for a special purpose DNA Reader type which serves as the output parameter of DNACalib.
+*/
+class DNACALIBMODULE_API FDNACalibDNAReader : public IDNAReader
+{
+public:
+	explicit FDNACalibDNAReader(IDNAReader* Source);
+
+	// Header
+	uint16 GetFileFormatGeneration() const override;
+	uint16 GetFileFormatVersion() const override;
+	// Descriptor
+	FString GetName() const override;
+	EArchetype GetArchetype() const override;
+	EGender GetGender() const override;
+	uint16 GetAge() const override;
+	uint32 GetMetaDataCount() const override;
+	FString GetMetaDataKey(uint32 Index) const override;
+	FString GetMetaDataValue(const FString& Key) const override;
+	ETranslationUnit GetTranslationUnit() const override;
+	ERotationUnit GetRotationUnit() const override;
+	FCoordinateSystem GetCoordinateSystem() const override;
+	uint16 GetLODCount() const override;
+	uint16 GetDBMaxLOD() const override;
+	FString GetDBComplexity() const override;
+	FString GetDBName() const override;
+	// Definition
+	uint16 GetGUIControlCount() const override;
+	FString GetGUIControlName(uint16 Index) const override;
+	uint16 GetRawControlCount() const override;
+	FString GetRawControlName(uint16 Index) const override;
+	uint16 GetJointCount() const override;
+	FString GetJointName(uint16 Index) const override;
+	uint16 GetJointIndexListCount() const override;
+	TArrayView<const uint16> GetJointIndicesForLOD(uint16 LOD) const override;
+	uint16 GetJointParentIndex(uint16 Index) const override;
+	uint16 GetBlendShapeChannelCount() const override;
+	FString GetBlendShapeChannelName(uint16 Index) const override;
+	uint16 GetBlendShapeChannelIndexListCount() const override;
+	TArrayView<const uint16> GetBlendShapeChannelIndicesForLOD(uint16 LOD) const override;
+	uint16 GetAnimatedMapCount() const override;
+	FString GetAnimatedMapName(uint16 Index) const override;
+	uint16 GetAnimatedMapIndexListCount() const override;
+	TArrayView<const uint16> GetAnimatedMapIndicesForLOD(uint16 LOD) const override;
+	uint16 GetMeshCount() const override;
+	FString GetMeshName(uint16 Index) const override;
+	uint16 GetMeshIndexListCount() const override;
+	TArrayView<const uint16> GetMeshIndicesForLOD(uint16 LOD) const override;
+	uint16 GetMeshBlendShapeChannelMappingCount() const override;
+	FMeshBlendShapeChannelMapping GetMeshBlendShapeChannelMapping(uint16 Index) const override;
+	TArrayView<const uint16> GetMeshBlendShapeChannelMappingIndicesForLOD(uint16 LOD) const override;
+	FVector GetNeutralJointTranslation(uint16 Index) const override;
+	FVector GetNeutralJointRotation(uint16 Index) const override;
+	// Behavior
+	TArrayView<const uint16> GetGUIToRawInputIndices() const override;
+	TArrayView<const uint16> GetGUIToRawOutputIndices() const override;
+	TArrayView<const float> GetGUIToRawFromValues() const override;
+	TArrayView<const float> GetGUIToRawToValues() const override;
+	TArrayView<const float> GetGUIToRawSlopeValues() const override;
+	TArrayView<const float> GetGUIToRawCutValues() const override;
+	uint16 GetPSDCount() const override;
+	TArrayView<const uint16> GetPSDRowIndices() const override;
+	TArrayView<const uint16> GetPSDColumnIndices() const override;
+	TArrayView<const float> GetPSDValues() const override;
+	uint16 GetJointRowCount() const override;
+	uint16 GetJointColumnCount() const override;
+	TArrayView<const uint16> GetJointVariableAttributeIndices(uint16 LOD) const override;
+	uint16 GetJointGroupCount() const override;
+	TArrayView<const uint16> GetJointGroupLODs(uint16 JointGroupIndex) const override;
+	TArrayView<const uint16> GetJointGroupInputIndices(uint16 JointGroupIndex) const override;
+	TArrayView<const uint16> GetJointGroupOutputIndices(uint16 JointGroupIndex) const override;
+	TArrayView<const float> GetJointGroupValues(uint16 JointGroupIndex) const override;
+	TArrayView<const uint16> GetJointGroupJointIndices(uint16 JointGroupIndex) const override;
+	TArrayView<const uint16> GetBlendShapeChannelLODs() const override;
+	TArrayView<const uint16> GetBlendShapeChannelOutputIndices() const override;
+	TArrayView<const uint16> GetBlendShapeChannelInputIndices() const override;
+	TArrayView<const uint16> GetAnimatedMapLODs() const override;
+	TArrayView<const uint16> GetAnimatedMapInputIndices() const override;
+	TArrayView<const uint16> GetAnimatedMapOutputIndices() const override;
+	TArrayView<const float> GetAnimatedMapFromValues() const override;
+	TArrayView<const float> GetAnimatedMapToValues() const override;
+	TArrayView<const float> GetAnimatedMapSlopeValues() const override;
+	TArrayView<const float> GetAnimatedMapCutValues() const override;
+	// Geometry
+	uint32 GetVertexPositionCount(uint16 MeshIndex) const override;
+	FVector GetVertexPosition(uint16 MeshIndex, uint32 VertexIndex) const override;
+	TArrayView<const float> GetVertexPositionXs(uint16 MeshIndex) const override;
+	TArrayView<const float> GetVertexPositionYs(uint16 MeshIndex) const override;
+	TArrayView<const float> GetVertexPositionZs(uint16 MeshIndex) const override;
+	uint32 GetVertexTextureCoordinateCount(uint16 MeshIndex) const override;
+	FTextureCoordinate GetVertexTextureCoordinate(uint16 MeshIndex, uint32 TextureCoordinateIndex) const override;
+	uint32 GetVertexNormalCount(uint16 MeshIndex) const override;
+	FVector GetVertexNormal(uint16 MeshIndex, uint32 NormalIndex) const override;
+	uint32 GetFaceCount(uint16 MeshIndex) const override;
+	TArrayView<const uint32> GetFaceVertexLayoutIndices(uint16 MeshIndex, uint32 FaceIndex) const override;
+	uint32 GetVertexLayoutCount(uint16 MeshIndex) const override;
+	FVertexLayout GetVertexLayout(uint16 MeshIndex, uint32 LayoutIndex) const override;
+	uint16 GetMaximumInfluencePerVertex(uint16 MeshIndex) const override;
+	uint32 GetSkinWeightsCount(uint16 MeshIndex) const override;
+	TArrayView<const float> GetSkinWeightsValues(uint16 MeshIndex, uint32 VertexIndex) const override;
+	TArrayView<const uint16> GetSkinWeightsJointIndices(uint16 MeshIndex, uint32 VertexIndex) const override;
+	uint16 GetBlendShapeTargetCount(uint16 MeshIndex) const override;
+	uint16 GetBlendShapeChannelIndex(uint16 MeshIndex, uint16 BlendShapeTargetIndex) const override;
+	uint32 GetBlendShapeTargetDeltaCount(uint16 MeshIndex, uint16 BlendShapeTargetIndex) const override;
+	FVector GetBlendShapeTargetDelta(uint16 MeshIndex, uint16 BlendShapeTargetIndex, uint32 DeltaIndex) const override;
+	TArrayView<const float> GetBlendShapeTargetDeltaXs(uint16 MeshIndex, uint16 BlendShapeTargetIndex) const override;
+	TArrayView<const float> GetBlendShapeTargetDeltaYs(uint16 MeshIndex, uint16 BlendShapeTargetIndex) const override;
+	TArrayView<const float> GetBlendShapeTargetDeltaZs(uint16 MeshIndex, uint16 BlendShapeTargetIndex) const override;
+	TArrayView<const uint32> GetBlendShapeTargetVertexIndices(uint16 MeshIndex, uint16 BlendShapeTargetIndex) const override;
+	// Machine Learned Behavior
+	uint16 GetMLControlCount() const override;
+	FString GetMLControlName(uint16 Index) const override;
+	uint16 GetNeuralNetworkCount() const override;
+	uint16 GetNeuralNetworkIndexListCount() const override;
+	TArrayView<const uint16> GetNeuralNetworkIndicesForLOD(uint16 LOD) const override;
+	uint16 GetMeshRegionCount(uint16 MeshIndex) const override;
+	FString GetMeshRegionName(uint16 MeshIndex, uint16 RegionIndex) const override;
+	TArrayView<const uint16> GetNeuralNetworkIndicesForMeshRegion(uint16 MeshIndex, uint16 RegionIndex) const override;
+	TArrayView<const uint16> GetNeuralNetworkInputIndices(uint16 NetIndex) const override;
+	TArrayView<const uint16> GetNeuralNetworkOutputIndices(uint16 NetIndex) const override;
+	uint16 GetNeuralNetworkLayerCount(uint16 NetIndex) const override;
+	EActivationFunction GetNeuralNetworkLayerActivationFunction(uint16 NetIndex, uint16 LayerIndex) const override;
+	TArrayView<const float> GetNeuralNetworkLayerActivationFunctionParameters(uint16 NetIndex, uint16 LayerIndex) const override;
+	TArrayView<const float> GetNeuralNetworkLayerBiases(uint16 NetIndex, uint16 LayerIndex) const override;
+	TArrayView<const float> GetNeuralNetworkLayerWeights(uint16 NetIndex, uint16 LayerIndex) const override;
+	// JointBehaviorMetadataReader
+	ETranslationRepresentation GetJointTranslationRepresentation(uint16 JointIndex) const override;
+	ERotationRepresentation GetJointRotationRepresentation(uint16 JointIndex) const override;
+	EScaleRepresentation GetJointScaleRepresentation(uint16 JointIndex) const override;
+	// RBFBehavior
+	uint16 GetRBFPoseCount() const override;
+	FString GetRBFPoseName(uint16 PoseIndex) const override;
+	TArrayView<const uint16> GetRBFPoseJointOutputIndices(uint16 PoseIndex) const override;
+	TArrayView<const uint16> GetRBFPoseBlendShapeChannelOutputIndices(uint16 PoseIndex) const override;
+	TArrayView<const uint16> GetRBFPoseAnimatedMapOutputIndices(uint16 PoseIndex) const override;
+	TArrayView<const float> GetRBFPoseJointOutputValues(uint16 PoseIndex) const override;
+	float GetRBFPoseScale(uint16 PoseIndex) const override;
+	uint16 GetRBFPoseControlCount() const override;
+	FString GetRBFPoseControlName(uint16 PoseControlIndex) const override;
+	TArrayView<const uint16> GetRBFPoseInputControlIndices(uint16 PoseIndex) const override;
+	TArrayView<const uint16> GetRBFPoseOutputControlIndices(uint16 PoseIndex) const override;
+	TArrayView<const float> GetRBFPoseOutputControlWeights(uint16 PoseIndex) const override;
+	uint16 GetRBFSolverCount() const override;
+	uint16 GetRBFSolverIndexListCount() const override;
+	TArrayView<const uint16> GetRBFSolverIndicesForLOD(uint16 LOD) const override;
+	FString GetRBFSolverName(uint16 SolverIndex) const override;
+	TArrayView<const uint16> GetRBFSolverRawControlIndices(uint16 SolverIndex) const override;
+	TArrayView<const uint16> GetRBFSolverPoseIndices(uint16 SolverIndex) const override;
+	TArrayView<const float> GetRBFSolverRawControlValues(uint16 SolverIndex) const override;
+	ERBFSolverType GetRBFSolverType(uint16 SolverIndex) const override;
+	float GetRBFSolverRadius(uint16 SolverIndex) const override;
+	EAutomaticRadius GetRBFSolverAutomaticRadius(uint16 SolverIndex) const override;
+	float GetRBFSolverWeightThreshold(uint16 SolverIndex) const override;
+	ERBFDistanceMethod GetRBFSolverDistanceMethod(uint16 SolverIndex) const override;
+	ERBFNormalizeMethod GetRBFSolverNormalizeMethod(uint16 SolverIndex) const override;
+	ERBFFunctionType GetRBFSolverFunctionType(uint16 SolverIndex) const override;
+	ETwistAxis GetRBFSolverTwistAxis(uint16 SolverIndex) const override;
+	// TwistSwingBehavior
+	uint16 GetTwistCount() const override;
+	ETwistAxis GetTwistSetupTwistAxis(uint16 TwistIndex) const override;
+	TArrayView<const uint16> GetTwistInputControlIndices(uint16 TwistIndex) const override;
+	TArrayView<const uint16> GetTwistOutputJointIndices(uint16 TwistIndex) const override;
+	TArrayView<const float> GetTwistBlendWeights(uint16 TwistIndex) const override;
+	uint16 GetSwingCount() const override;
+	ETwistAxis GetSwingSetupTwistAxis(uint16 SwingIndex) const override;
+	TArrayView<const uint16> GetSwingInputControlIndices(uint16 SwingIndex) const override;
+	TArrayView<const uint16> GetSwingOutputJointIndices(uint16 SwingIndex) const override;
+	TArrayView<const float> GetSwingBlendWeights(uint16 SwingIndex) const override;
+
+	void Unload(EDNADataLayer layer) override;
+
+	dna::Reader* Unwrap() const override;
+
+private:
+	TUniquePtr<IDNAReader> ReaderPtr;
+
+};
diff --git a/Engine/Plugins/Animation/DNACalib/Source/DNACalibModule/Public/DNACalibModule.h b/Engine/Plugins/Animation/DNACalib/Source/DNACalibModule/Public/DNACalibModule.h
new file mode 100644
--- /dev/null
+++ b/Engine/Plugins/Animation/DNACalib/Source/DNACalibModule/Public/DNACalibModule.h
@@ -0,0 +1,15 @@
+// Copyright Epic Games, Inc. All Rights Reserved.
+
+#pragma once
+
+#include "CoreMinimal.h"
+#include "Modules/ModuleManager.h"
+
+class DNACALIBMODULE_API FDNACalibModule : public IModuleInterface
+{
+public:
+
+	/** IModuleInterface implementation */
+	virtual void StartupModule() override;
+	virtual void ShutdownModule() override;
+};
