[DNACalibLibTest]
* Changed so test module does not add DNACalibLib code if building with merged modules since we know they end up in the same dll

diff --git a/Engine/Plugins/Animation/DNACalib/Source/DNACalibLibTest/DNACalibLibTest.Build.cs b/Engine/Plugins/Animation/DNACalib/Source/DNACalibLibTest/DNACalibLibTest.Build.cs
--- a/Engine/Plugins/Animation/DNACalib/Source/DNACalibLibTest/DNACalibLibTest.Build.cs
+++ b/Engine/Plugins/Animation/DNACalib/Source/DNACalibLibTest/DNACalibLibTest.Build.cs
@@ -18,7 +18,7 @@ namespace UnrealBuildTool.Rules
 
 			string DNACalibLibPath = Path.GetFullPath(Path.Combine(ModuleDirectory, "../DNACalibLib"));
 
-			if (Target.LinkType == TargetLinkType.Monolithic)
+			if (Target.LinkType == TargetLinkType.Monolithic || Target.bMergeModules)
 			{
 				PublicDependencyModuleNames.Add("DNACalibLib");
 				PrivateIncludePaths.Add(Path.Combine(DNACalibLibPath, "Private"));
