[DNACalib] Remove FMessageLog as a dependency from DNACalibModule

- FMessageLog is an editor only module and its not used by DNACalib and causes errors when packaging the shipping configuration

[RN] minor

diff --git a/Engine/Plugins/Animation/DNACalib/Source/DNACalibModule/DNACalibModule.Build.cs b/Engine/Plugins/Animation/DNACalib/Source/DNACalibModule/DNACalibModule.Build.cs
--- a/Engine/Plugins/Animation/DNACalib/Source/DNACalibModule/DNACalibModule.Build.cs
+++ b/Engine/Plugins/Animation/DNACalib/Source/DNACalibModule/DNACalibModule.Build.cs
@@ -13,7 +13,6 @@ namespace UnrealBuildTool.Rules
 					"CoreUObject",
 					"Engine",
 					"RigLogicModule",
-					"MessageLog"
 				} 
 			);
 
