Disabled iwyu support for DNACalid lib

diff --git a/Engine/Plugins/Animation/DNACalib/Source/DNACalibLib/DNACalibLib.Build.cs b/Engine/Plugins/Animation/DNACalib/Source/DNACalibLib/DNACalibLib.Build.cs
--- a/Engine/Plugins/Animation/DNACalib/Source/DNACalibLib/DNACalibLib.Build.cs
+++ b/Engine/Plugins/Animation/DNACalib/Source/DNACalibLib/DNACalibLib.Build.cs
@@ -6,6 +6,8 @@ namespace UnrealBuildTool.Rules
 	{
 		public DNACalibLib(ReadOnlyTargetRules Target) : base(Target)
 		{
+			IWYUSupport = IWYUSupport.None;
+
 			PublicDependencyModuleNames.AddRange(
 				new string[]
 				{
