Fix link error caused by missing declspec in update to DNACalib
#Synced-CL 49330023

diff --git a/Engine/Plugins/Animation/DNACalib/Source/DNACalibModule/Public/DNACalibDNAReader.h b/Engine/Plugins/Animation/DNACalib/Source/DNACalibModule/Public/DNACalibDNAReader.h
--- a/Engine/Plugins/Animation/DNACalib/Source/DNACalibModule/Public/DNACalibDNAReader.h
+++ b/Engine/Plugins/Animation/DNACalib/Source/DNACalibModule/Public/DNACalibDNAReader.h
@@ -180,19 +180,19 @@ public:
 	UE_API TArrayView<const uint16> GetSwingOutputJointIndices(uint16 SwingIndex) const override;
 	UE_API TArrayView<const float> GetSwingBlendWeights(uint16 SwingIndex) const override;
 	// MachineLearnedBehaviorExtReader
-	uint16 GetMLTypeCount() const override;
-	uint16 GetMLOperationSetCount(uint16 MLTypeIndex) const override;
-	uint16 GetMLOperationCount(uint16 MLTypeIndex, uint16 MLOperationSetIndex) const override;
-	EMachineLearnedBehaviorOperationType GetMLOperationType(uint16 MLTypeIndex, uint16 MLOperationSetIndex, uint16 MLOperationIndex) const override;
-	TArrayView<const uint32> GetMLOperationParameters(uint16 MLTypeIndex, uint16 MLOperationSetIndex, uint16 MLOperationIndex) const override;
-	TArrayView<const uint16> GetMLOperationDependencyOperationSetIndices(uint16 MLTypeIndex, uint16 MLOperationSetIndex, uint16 MLOperationIndex) const override;
-	TArrayView<const uint16> GetMLOperationDependencyOperationIndices(uint16 MLTypeIndex, uint16 MLOperationSetIndex, uint16 MLOperationIndex) const override;
-	uint16 GetMLOperationIndexListCount(uint16 MLTypeIndex, uint16 MLOperationSetIndex) const override;
-	TArrayView<const uint16> GetMLOperationIndicesForLOD(uint16 MLTypeIndex, uint16 MLOperationSetIndex, uint16 LOD) const override;
-	TArrayView<const uint16> GetMLJointsInputIndices() const override;
-	TArrayView<const uint16> GetMLJointsOutputIndices() const override;
-	TArrayView<const uint16> GetMLJointsParameterKeys() const override;
-	TArrayView<const uint16> GetMLJointsParameterValues() const override;
+	UE_API uint16 GetMLTypeCount() const override;
+	UE_API uint16 GetMLOperationSetCount(uint16 MLTypeIndex) const override;
+	UE_API uint16 GetMLOperationCount(uint16 MLTypeIndex, uint16 MLOperationSetIndex) const override;
+	UE_API EMachineLearnedBehaviorOperationType GetMLOperationType(uint16 MLTypeIndex, uint16 MLOperationSetIndex, uint16 MLOperationIndex) const override;
+	UE_API TArrayView<const uint32> GetMLOperationParameters(uint16 MLTypeIndex, uint16 MLOperationSetIndex, uint16 MLOperationIndex) const override;
+	UE_API TArrayView<const uint16> GetMLOperationDependencyOperationSetIndices(uint16 MLTypeIndex, uint16 MLOperationSetIndex, uint16 MLOperationIndex) const override;
+	UE_API TArrayView<const uint16> GetMLOperationDependencyOperationIndices(uint16 MLTypeIndex, uint16 MLOperationSetIndex, uint16 MLOperationIndex) const override;
+	UE_API uint16 GetMLOperationIndexListCount(uint16 MLTypeIndex, uint16 MLOperationSetIndex) const override;
+	UE_API TArrayView<const uint16> GetMLOperationIndicesForLOD(uint16 MLTypeIndex, uint16 MLOperationSetIndex, uint16 LOD) const override;
+	UE_API TArrayView<const uint16> GetMLJointsInputIndices() const override;
+	UE_API TArrayView<const uint16> GetMLJointsOutputIndices() const override;
+	UE_API TArrayView<const uint16> GetMLJointsParameterKeys() const override;
+	UE_API TArrayView<const uint16> GetMLJointsParameterValues() const override;
 	UE_API void Unload(EDNADataLayer layer) override;
 
 	UE_API dna::Reader* Unwrap() const override;
